Live Link Sound Devices Plugin - Port from Switchboard to Live Link Device framework

Key Highlights:
- Full REST API integration
- Mock server for testing without hardware

Tested with mock server and recording, connecting, and take / slate interactions worked.

#Synced-CL 51338485

diff --git a/Engine/Plugins/VirtualProduction/LiveLinkSoundDevice/LiveLinkSoundDevice.uplugin b/Engine/Plugins/VirtualProduction/LiveLinkSoundDevice/LiveLinkSoundDevice.uplugin
new file mode 100644
--- /dev/null
+++ b/Engine/Plugins/VirtualProduction/LiveLinkSoundDevice/LiveLinkSoundDevice.uplugin
@@ -0,0 +1,27 @@
+{
+	"FileVersion": 3,
+	"Version": 1,
+	"VersionName": "1.0",
+	"FriendlyName": "Live Link Sound Device",
+	"Description": "Sound Devices recorder support for Live Link Hub with recording and connection capabilities",
+	"Category": "Virtual Production",
+	"CreatedBy": "Epic Games, Inc.",
+	"CreatedByURL": "https://epicgames.com",
+	"CanContainContent": true,
+	"IsBetaVersion": false,
+	"IsExperimentalVersion": true,
+	"Installed": false,
+	"Modules": [
+		{
+			"Name": "LiveLinkSoundDevice",
+			"Type": "Editor",
+			"LoadingPhase": "Default"
+		}
+	],
+	"Plugins": [
+		{
+			"Name": "LiveLinkDevice",
+			"Enabled": true
+		}
+	]
+}
diff --git a/Engine/Plugins/VirtualProduction/LiveLinkSoundDevice/Source/LiveLinkSoundDevice/LiveLinkSoundDevice.Build.cs b/Engine/Plugins/VirtualProduction/LiveLinkSoundDevice/Source/LiveLinkSoundDevice/LiveLinkSoundDevice.Build.cs
new file mode 100644
--- /dev/null
+++ b/Engine/Plugins/VirtualProduction/LiveLinkSoundDevice/Source/LiveLinkSoundDevice/LiveLinkSoundDevice.Build.cs
@@ -0,0 +1,35 @@
+// Copyright Epic Games, Inc. All Rights Reserved.
+
+using UnrealBuildTool;
+
+public class LiveLinkSoundDevice : ModuleRules
+{
+	public LiveLinkSoundDevice(ReadOnlyTargetRules Target) : base(Target)
+	{
+		PCHUsage = ModuleRules.PCHUsageMode.UseExplicitOrSharedPCHs;
+
+		PublicDependencyModuleNames.AddRange(
+			new string[]
+			{
+				"Core",
+			}
+		);
+
+		PrivateDependencyModuleNames.AddRange(
+			new string[]
+			{
+				"CoreUObject",
+				"Engine",
+				"Slate",
+				"SlateCore",
+				"LiveLinkDevice",
+				"LiveLinkInterface",
+				"Projects",
+				"HTTP",
+				"Json",
+				"JsonUtilities",
+				"Networking"
+			}
+		);
+	}
+}
diff --git a/Engine/Plugins/VirtualProduction/LiveLinkSoundDevice/Source/LiveLinkSoundDevice/Private/HttpDigestAuth.cpp b/Engine/Plugins/VirtualProduction/LiveLinkSoundDevice/Source/LiveLinkSoundDevice/Private/HttpDigestAuth.cpp
new file mode 100644
--- /dev/null
+++ b/Engine/Plugins/VirtualProduction/LiveLinkSoundDevice/Source/LiveLinkSoundDevice/Private/HttpDigestAuth.cpp
@@ -0,0 +1,181 @@
+// Copyright Epic Games, Inc. All Rights Reserved.
+
+#include "HttpDigestAuth.h"
+#include "Misc/SecureHash.h"
+#include "Misc/Guid.h"
+
+bool FHttpDigestAuthHelper::ParseDigestChallenge(const FString& AuthHeader, FDigestChallenge& OutChallenge)
+{
+	// WWW-Authenticate: Digest realm="...", nonce="...", qop="auth", ...
+
+	// Remove "Digest " prefix if present
+	FString Header = AuthHeader;
+	if (Header.StartsWith(TEXT("Digest "), ESearchCase::IgnoreCase))
+	{
+		Header = Header.RightChop(7).TrimStart();
+	}
+
+	// Split by comma (but be careful of quoted values)
+	TArray<FString> Parts;
+	int32 Start = 0;
+	bool bInQuotes = false;
+
+	for (int32 i = 0; i < Header.Len(); i++)
+	{
+		if (Header[i] == '"')
+		{
+			bInQuotes = !bInQuotes;
+		}
+		else if (Header[i] == ',' && !bInQuotes)
+		{
+			Parts.Add(Header.Mid(Start, i - Start).TrimStartAndEnd());
+			Start = i + 1;
+		}
+	}
+	// Add last part
+	if (Start < Header.Len())
+	{
+		Parts.Add(Header.Mid(Start).TrimStartAndEnd());
+	}
+
+	// Parse each key=value pair
+	for (const FString& Part : Parts)
+	{
+		FString Key, Value;
+		if (ParseHeaderPair(Part, Key, Value))
+		{
+			if (Key.Equals(TEXT("realm"), ESearchCase::IgnoreCase))
+			{
+				OutChallenge.Realm = Value;
+			}
+			else if (Key.Equals(TEXT("nonce"), ESearchCase::IgnoreCase))
+			{
+				OutChallenge.Nonce = Value;
+			}
+			else if (Key.Equals(TEXT("qop"), ESearchCase::IgnoreCase))
+			{
+				OutChallenge.Qop = Value;
+			}
+			else if (Key.Equals(TEXT("opaque"), ESearchCase::IgnoreCase))
+			{
+				OutChallenge.Opaque = Value;
+			}
+			else if (Key.Equals(TEXT("algorithm"), ESearchCase::IgnoreCase))
+			{
+				OutChallenge.Algorithm = Value;
+			}
+		}
+	}
+
+	// Realm and Nonce are required
+	return !OutChallenge.Realm.IsEmpty() && !OutChallenge.Nonce.IsEmpty();
+}
+
+FString FHttpDigestAuthHelper::GenerateDigestResponse(
+	const FString& Username,
+	const FString& Password,
+	const FString& Method,
+	const FString& URI,
+	const FDigestChallenge& Challenge)
+{
+	// HA1 = MD5(username:realm:password)
+	FString HA1Input = FString::Printf(TEXT("%s:%s:%s"), *Username, *Challenge.Realm, *Password);
+	FString HA1 = ComputeMD5Hash(HA1Input);
+
+	// HA2 = MD5(method:uri)
+	FString HA2Input = FString::Printf(TEXT("%s:%s"), *Method, *URI);
+	FString HA2 = ComputeMD5Hash(HA2Input);
+
+	FString Response;
+	FString AuthValue;
+
+	// Check if qop is specified
+	if (!Challenge.Qop.IsEmpty())
+	{
+		// qop-based authentication (RFC 2617)
+		FString CNonce = GenerateNonce();
+		FString NonceCount = TEXT("00000001");
+
+		// Response = MD5(HA1:nonce:nc:cnonce:qop:HA2)
+		FString ResponseInput = FString::Printf(
+			TEXT("%s:%s:%s:%s:%s:%s"),
+			*HA1, *Challenge.Nonce, *NonceCount, *CNonce, *Challenge.Qop, *HA2);
+		Response = ComputeMD5Hash(ResponseInput);
+
+		// Build Authorization header
+		AuthValue = FString::Printf(
+			TEXT("Digest username=\"%s\", realm=\"%s\", nonce=\"%s\", uri=\"%s\", ")
+			TEXT("qop=%s, nc=%s, cnonce=\"%s\", response=\"%s\""),
+			*Username, *Challenge.Realm, *Challenge.Nonce, *URI,
+			*Challenge.Qop, *NonceCount, *CNonce, *Response);
+	}
+	else
+	{
+		// Legacy mode without qop
+		// Response = MD5(HA1:nonce:HA2)
+		FString ResponseInput = FString::Printf(TEXT("%s:%s:%s"), *HA1, *Challenge.Nonce, *HA2);
+		Response = ComputeMD5Hash(ResponseInput);
+
+		// Build Authorization header
+		AuthValue = FString::Printf(
+			TEXT("Digest username=\"%s\", realm=\"%s\", nonce=\"%s\", uri=\"%s\", response=\"%s\""),
+			*Username, *Challenge.Realm, *Challenge.Nonce, *URI, *Response);
+	}
+
+	// Add opaque if present
+	if (!Challenge.Opaque.IsEmpty())
+	{
+		AuthValue += FString::Printf(TEXT(", opaque=\"%s\""), *Challenge.Opaque);
+	}
+
+	return AuthValue;
+}
+
+FString FHttpDigestAuthHelper::ComputeMD5Hash(const FString& Input)
+{
+	// Use Unreal's FMD5 for hash computation
+	FMD5 MD5Gen;
+
+	// Convert to UTF-8 for hashing
+	FTCHARToUTF8 UTF8String(*Input);
+	MD5Gen.Update(reinterpret_cast<const uint8*>(UTF8String.Get()), UTF8String.Length());
+
+	uint8 Digest[16];
+	MD5Gen.Final(Digest);
+
+	// Convert to lowercase hex string
+	FString Result;
+	for (int32 i = 0; i < 16; i++)
+	{
+		Result += FString::Printf(TEXT("%02x"), Digest[i]);
+	}
+
+	return Result;
+}
+
+FString FHttpDigestAuthHelper::GenerateNonce()
+{
+	// Generate random nonce using GUID
+	FGuid Guid = FGuid::NewGuid();
+	return Guid.ToString(EGuidFormats::Digits);
+}
+
+bool FHttpDigestAuthHelper::ParseHeaderPair(const FString& Pair, FString& OutKey, FString& OutValue)
+{
+	int32 EqualPos;
+	if (!Pair.FindChar('=', EqualPos))
+	{
+		return false;
+	}
+
+	OutKey = Pair.Left(EqualPos).TrimStartAndEnd();
+	OutValue = Pair.Mid(EqualPos + 1).TrimStartAndEnd();
+
+	// Remove quotes if present
+	if (OutValue.StartsWith(TEXT("\"")) && OutValue.EndsWith(TEXT("\"")))
+	{
+		OutValue = OutValue.Mid(1, OutValue.Len() - 2);
+	}
+
+	return true;
+}
diff --git a/Engine/Plugins/VirtualProduction/LiveLinkSoundDevice/Source/LiveLinkSoundDevice/Private/LiveLinkSoundDeviceBase.cpp b/Engine/Plugins/VirtualProduction/LiveLinkSoundDevice/Source/LiveLinkSoundDevice/Private/LiveLinkSoundDeviceBase.cpp
new file mode 100644
--- /dev/null
+++ b/Engine/Plugins/VirtualProduction/LiveLinkSoundDevice/Source/LiveLinkSoundDevice/Private/LiveLinkSoundDeviceBase.cpp
@@ -0,0 +1,695 @@
+// Copyright Epic Games, Inc. All Rights Reserved.
+
+#include "LiveLinkSoundDeviceBase.h"
+#include "HttpModule.h"
+#include "Interfaces/IHttpResponse.h"
+#include "Dom/JsonValue.h"
+#include "Dom/JsonObject.h"
+#include "Serialization/JsonReader.h"
+#include "Serialization/JsonSerializer.h"
+#include "LiveLinkHub/ILiveLinkRecordingSession.h"
+#include "Interfaces/IPv4/IPv4Endpoint.h"
+#include "GenericPlatform/GenericPlatformHttp.h"
+#include "Misc/DateTime.h"
+
+#define LOCTEXT_NAMESPACE "LiveLinkSoundDeviceRecorder"
+
+//~ Begin ULiveLinkDevice interface
+
+TSubclassOf<ULiveLinkDeviceSettings> ULiveLinkSoundDeviceBase::GetSettingsClass() const
+{
+	return SettingsClass ? SettingsClass : TSubclassOf<ULiveLinkDeviceSettings>(ULiveLinkSoundDeviceSettings::StaticClass());
+}
+
+FText ULiveLinkSoundDeviceBase::GetDisplayName() const
+{
+	const ULiveLinkSoundDeviceSettings* DeviceSettings = GetDeviceSettings<ULiveLinkSoundDeviceSettings>();
+	return FText::FromString(FString::Printf(TEXT("Sound Devices (%s)"), *DeviceSettings->IpAddress));
+}
+
+EDeviceHealth ULiveLinkSoundDeviceBase::GetDeviceHealth() const
+{
+	switch (ConnectionStatus)
+	{
+	case ELiveLinkDeviceConnectionStatus::Connected:
+		return bIsRecording ? EDeviceHealth::Good : EDeviceHealth::Nominal;
+	case ELiveLinkDeviceConnectionStatus::Connecting:
+		return EDeviceHealth::Info;
+	case ELiveLinkDeviceConnectionStatus::Disconnecting:
+		return EDeviceHealth::Warning;
+	case ELiveLinkDeviceConnectionStatus::Disconnected:
+	default:
+		return EDeviceHealth::Error;
+	}
+}
+
+FText ULiveLinkSoundDeviceBase::GetHealthText() const
+{
+	switch (ConnectionStatus)
+	{
+	case ELiveLinkDeviceConnectionStatus::Connected:
+		return bIsRecording ? LOCTEXT("Recording", "Recording") : LOCTEXT("Connected", "Connected - Ready");
+	case ELiveLinkDeviceConnectionStatus::Connecting:
+		return LOCTEXT("Connecting", "Connecting...");
+	case ELiveLinkDeviceConnectionStatus::Disconnecting:
+		return LOCTEXT("Disconnecting", "Disconnecting...");
+	case ELiveLinkDeviceConnectionStatus::Disconnected:
+	default:
+		return LOCTEXT("Disconnected", "Disconnected");
+	}
+}
+
+void ULiveLinkSoundDeviceBase::OnDeviceAdded()
+{
+	// Subscribe to recording session delegates for slate/take changes
+	ILiveLinkRecordingSession& Session = ILiveLinkRecordingSession::Get();
+	Session.OnSlateNameChanged().AddUObject(this, &ULiveLinkSoundDeviceBase::HandleSlateNameChanged);
+	Session.OnTakeNumberChanged().AddUObject(this, &ULiveLinkSoundDeviceBase::HandleTakeNumberChanged);
+
+	// Start main device ticker (handles reconnection and polling)
+	DeviceTickerHandle = FTSTicker::GetCoreTicker().AddTicker(
+		FTickerDelegate::CreateUObject(this, &ULiveLinkSoundDeviceBase::DeviceTick),
+		PollingIntervalSeconds
+	);
+}
+
+void ULiveLinkSoundDeviceBase::OnDeviceRemoved()
+{
+	// Unsubscribe from recording session delegates
+	ILiveLinkRecordingSession& Session = ILiveLinkRecordingSession::Get();
+	Session.OnSlateNameChanged().RemoveAll(this);
+	Session.OnTakeNumberChanged().RemoveAll(this);
+
+	// Stop device ticker
+	if (DeviceTickerHandle.IsValid())
+	{
+		FTSTicker::GetCoreTicker().RemoveTicker(DeviceTickerHandle);
+		DeviceTickerHandle.Reset();
+	}
+
+	if (ConnectionStatus != ELiveLinkDeviceConnectionStatus::Disconnected)
+	{
+		Disconnect_Implementation();
+	}
+}
+
+void ULiveLinkSoundDeviceBase::OnSettingChanged(const FPropertyChangedEvent& InPropertyChangedEvent)
+{
+	// When settings change, attempt immediate reconnection if disconnected
+	if (ConnectionStatus == ELiveLinkDeviceConnectionStatus::Disconnected)
+	{
+		UE_LOG(LogTemp, Log, TEXT("Settings changed, attempting connection to Sound Devices..."));
+		AttemptReconnection();
+	}
+}
+
+//~ End ULiveLinkDevice interface
+
+//~ Begin ILiveLinkDeviceCapability_Connection interface
+
+ELiveLinkDeviceConnectionStatus ULiveLinkSoundDeviceBase::GetConnectionStatus_Implementation() const
+{
+	return ConnectionStatus;
+}
+
+FString ULiveLinkSoundDeviceBase::GetHardwareId_Implementation() const
+{
+	const ULiveLinkSoundDeviceSettings* DeviceSettings = GetDeviceSettings<ULiveLinkSoundDeviceSettings>();
+	return FString::Printf(TEXT("%s:%d"), *DeviceSettings->IpAddress, DeviceSettings->Port);
+}
+
+bool ULiveLinkSoundDeviceBase::SetHardwareId_Implementation(const FString& HardwareID)
+{
+	ULiveLinkSoundDeviceSettings* DeviceSettings = GetDeviceSettings<ULiveLinkSoundDeviceSettings>();
+
+	FIPv4Endpoint OutEndpoint;
+	if (FIPv4Endpoint::Parse(HardwareID, OutEndpoint))
+	{
+		DeviceSettings->IpAddress = OutEndpoint.Address.ToString();
+		DeviceSettings->Port = OutEndpoint.Port;
+	}
+
+	return true;
+}
+
+bool ULiveLinkSoundDeviceBase::CanSetHardwareId_Implementation()
+{
+	return ConnectionStatus == ELiveLinkDeviceConnectionStatus::Disconnected;
+}
+
+bool ULiveLinkSoundDeviceBase::Connect_Implementation()
+{
+	if (ConnectionStatus != ELiveLinkDeviceConnectionStatus::Disconnected)
+	{
+		return false;
+	}
+
+	ConnectionStatus = ELiveLinkDeviceConnectionStatus::Connecting;
+	SetConnectionStatus(ConnectionStatus);
+
+	// Test connection with timecode command
+	SendCommand(TEXT("tmcode"), [this](bool bSuccess, const FString& Response)
+	{
+		if (!bSuccess)
+		{
+			UE_LOG(LogTemp, Warning, TEXT("Failed to connect to Sound Devices"));
+			ConnectionStatus = ELiveLinkDeviceConnectionStatus::Disconnected;
+			SetConnectionStatus(ConnectionStatus);
+			return;
+		}
+
+		UE_LOG(LogTemp, Log, TEXT("Connected to Sound Devices, timecode: %s"), *Response);
+
+		// Set filename format to Scene-Take
+		SetSetting(TEXT("FileNameFormat"), TEXT("Scene-Take"), [](bool bSuccess) {});
+
+		// Set reel to today's date
+		FString DateStr = FDateTime::Now().ToString(TEXT("%Y%m%d"));
+		SetSetting(TEXT("ReelName"), DateStr, [](bool bSuccess) {});
+
+		// Configure all drives for recording
+		ConfigureDrivesForRecording();
+
+		// Mark as connected
+		ConnectionStatus = ELiveLinkDeviceConnectionStatus::Connected;
+		SetConnectionStatus(ConnectionStatus);
+
+		UE_LOG(LogTemp, Log, TEXT("Sound Devices configured and ready"));
+	});
+
+	return true;
+}
+
+bool ULiveLinkSoundDeviceBase::Disconnect_Implementation()
+{
+	if (ConnectionStatus == ELiveLinkDeviceConnectionStatus::Disconnected)
+	{
+		return false;
+	}
+
+	ConnectionStatus = ELiveLinkDeviceConnectionStatus::Disconnecting;
+	SetConnectionStatus(ConnectionStatus);
+
+	// Stop any active recording
+	if (bIsRecording)
+	{
+		StopRecording_Implementation();
+	}
+
+	// Clear polling state
+	CurrentPollingOperation = EPollingOperation::None;
+	bPollingRequestInFlight = false;
+
+	// Clear cached auth
+	bHasCachedChallenge = false;
+
+	ConnectionStatus = ELiveLinkDeviceConnectionStatus::Disconnected;
+	SetConnectionStatus(ConnectionStatus);
+
+	UE_LOG(LogTemp, Log, TEXT("Disconnected from Sound Devices"));
+
+	return true;
+}
+
+//~ End ILiveLinkDeviceCapability_Connection interface
+
+//~ Begin ILiveLinkDeviceCapability_Recording interface
+
+bool ULiveLinkSoundDeviceBase::StartRecording_Implementation()
+{
+	if (ConnectionStatus != ELiveLinkDeviceConnectionStatus::Connected || bIsRecording)
+	{
+		return false;
+	}
+
+	ILiveLinkRecordingSession& Session = ILiveLinkRecordingSession::Get();
+	const FString Slate = Session.GetSlateName();
+	const int32 Take = Session.GetTakeNumber();
+
+	UE_LOG(LogTemp, Log, TEXT("Starting recording: %s Take %d"), *Slate, Take);
+
+	// Set reel to today's date
+	FString DateStr = FDateTime::Now().ToString(TEXT("%Y%m%d"));
+	SetSetting(TEXT("ReelName"), DateStr, [](bool) {});
+
+	// Set slate and take
+	SetSetting(TEXT("SceneName"), Slate, [](bool) {});
+	SetSetting(TEXT("TakeNumber"), FString::FromInt(Take), [](bool) {});
+
+	// Ensure drives are in record mode
+	ConfigureDrivesForRecording();
+
+	// Send record command
+	SendCommand(TEXT("settransport/rec"), [](bool, const FString&) {});
+
+	// Start polling for "rec" state
+	CurrentPollingOperation = EPollingOperation::WaitingForRecording;
+	PollingStartTime = FPlatformTime::Seconds();
+	bPollingRequestInFlight = false;
+
+	return true;
+}
+
+bool ULiveLinkSoundDeviceBase::StopRecording_Implementation()
+{
+	if (ConnectionStatus != ELiveLinkDeviceConnectionStatus::Connected || !bIsRecording)
+	{
+		return false;
+	}
+
+	UE_LOG(LogTemp, Log, TEXT("Stopping recording..."));
+
+	// Send stop command
+	SendCommand(TEXT("settransport/stop"), [](bool, const FString&) {});
+
+	// Start polling for "stop" state
+	CurrentPollingOperation = EPollingOperation::WaitingForIdle;
+	PollingStartTime = FPlatformTime::Seconds();
+	bPollingRequestInFlight = false;
+
+	return true;
+}
+
+bool ULiveLinkSoundDeviceBase::IsRecording_Implementation() const
+{
+	return bIsRecording;
+}
+
+//~ End ILiveLinkDeviceCapability_Recording interface
+
+//~ Private Methods
+
+FString ULiveLinkSoundDeviceBase::GetBaseUrl() const
+{
+	const ULiveLinkSoundDeviceSettings* DeviceSettings = GetDeviceSettings<ULiveLinkSoundDeviceSettings>();
+	return FString::Printf(TEXT("http://%s:%d/sounddevices"), *DeviceSettings->IpAddress, DeviceSettings->Port);
+}
+
+FString ULiveLinkSoundDeviceBase::GetCommandUrl(const FString& Command) const
+{
+	return GetBaseUrl() + TEXT("/") + Command;
+}
+
+FString ULiveLinkSoundDeviceBase::GetRequestURI(const FString& Command) const
+{
+	return TEXT("/sounddevices/") + Command;
+}
+
+TSharedRef<IHttpRequest> ULiveLinkSoundDeviceBase::CreateRequest(const FString& Command)
+{
+	// Sound devices are usually local on the network so 1.0 should be a very long time.
+	const float HttpRequestTimeout = 1.0f;
+	TSharedRef<IHttpRequest> Request = FHttpModule::Get().CreateRequest();
+	Request->SetURL(GetCommandUrl(Command));
+	Request->SetVerb(TEXT("GET"));
+	Request->SetTimeout(HttpRequestTimeout);
+	return Request;
+}
+
+void ULiveLinkSoundDeviceBase::SendCommand(const FString& Command, TFunction<void(bool bSuccess, const FString& Response)> OnComplete)
+{
+	TWeakObjectPtr<ULiveLinkSoundDeviceBase> WeakThis(this);
+
+	// If we have a cached challenge, use it directly
+	if (bHasCachedChallenge)
+	{
+		RetryWithAuth(Command, CachedChallenge, OnComplete);
+		return;
+	}
+
+	// Phase 1: Initial request (expect 401)
+	TSharedRef<IHttpRequest> Request = CreateRequest(Command);
+	Request->OnProcessRequestComplete().BindLambda(
+		[WeakThis, Command, OnComplete](FHttpRequestPtr Req, FHttpResponsePtr Resp, bool bWasSuccessful)
+		{
+			if (!WeakThis.IsValid())
+			{
+				return;
+			}
+
+			if (bWasSuccessful && Resp.IsValid())
+			{
+				if (Resp->GetResponseCode() == 401)
+				{
+					// Phase 2: Parse challenge and retry with auth
+					FString AuthHeader = Resp->GetHeader(TEXT("WWW-Authenticate"));
+					FDigestChallenge Challenge;
+
+					if (FHttpDigestAuthHelper::ParseDigestChallenge(AuthHeader, Challenge))
+					{
+						// Cache challenge for future requests
+						WeakThis->CachedChallenge = Challenge;
+						WeakThis->bHasCachedChallenge = true;
+
+						WeakThis->RetryWithAuth(Command, Challenge, OnComplete);
+					}
+					else
+					{
+						UE_LOG(LogTemp, Warning, TEXT("Failed to parse Digest challenge"));
+						if (OnComplete)
+						{
+							OnComplete(false, FString());
+						}
+					}
+				}
+				else if (Resp->GetResponseCode() == 200)
+				{
+					// Success without auth (unlikely for Sound Devices)
+					if (OnComplete)
+					{
+						OnComplete(true, Resp->GetContentAsString());
+					}
+				}
+				else
+				{
+					UE_LOG(LogTemp, Warning, TEXT("HTTP error: %d"), Resp->GetResponseCode());
+					if (OnComplete)
+					{
+						OnComplete(false, FString());
+					}
+				}
+			}
+			else
+			{
+				UE_LOG(LogTemp, Warning, TEXT("HTTP request failed for command: %s"), *Command);
+				if (OnComplete)
+				{
+					OnComplete(false, FString());
+				}
+			}
+		});
+
+	Request->ProcessRequest();
+}
+
+void ULiveLinkSoundDeviceBase::RetryWithAuth(const FString& Command, const FDigestChallenge& Challenge, TFunction<void(bool, const FString&)> OnComplete)
+{
+	const ULiveLinkSoundDeviceSettings* DeviceSettings = GetDeviceSettings<ULiveLinkSoundDeviceSettings>();
+	TWeakObjectPtr<ULiveLinkSoundDeviceBase> WeakThis(this);
+
+	TSharedRef<IHttpRequest> Request = CreateRequest(Command);
+
+	FString AuthValue = FHttpDigestAuthHelper::GenerateDigestResponse(
+		DeviceSettings->Username,
+		DeviceSettings->Password,
+		Request->GetVerb(),
+		GetRequestURI(Command),
+		Challenge);
+
+	Request->SetHeader(TEXT("Authorization"), AuthValue);
+	Request->OnProcessRequestComplete().BindLambda(
+		[WeakThis, OnComplete](FHttpRequestPtr Req, FHttpResponsePtr Resp, bool bWasSuccessful)
+		{
+			if (!WeakThis.IsValid())
+			{
+				return;
+			}
+
+			if (bWasSuccessful && Resp.IsValid() && Resp->GetResponseCode() == 200)
+			{
+				if (OnComplete)
+				{
+					OnComplete(true, Resp->GetContentAsString());
+				}
+			}
+			else
+			{
+				if (Resp.IsValid())
+				{
+					UE_LOG(LogTemp, Warning, TEXT("Authenticated request failed: %d"), Resp->GetResponseCode());
+				}
+				if (OnComplete)
+				{
+					OnComplete(false, FString());
+				}
+			}
+		});
+
+	Request->ProcessRequest();
+}
+
+void ULiveLinkSoundDeviceBase::GetTransportState(TFunction<void(bool bSuccess, const FString& StateText)> OnComplete)
+{
+	SendCommand(TEXT("transport"), [OnComplete](bool bSuccess, const FString& Response)
+	{
+		if (!bSuccess)
+		{
+			if (OnComplete)
+			{
+				OnComplete(false, FString());
+			}
+			return;
+		}
+
+		// Parse JSON: {"Transport": "rec"|"stop"}
+		TSharedPtr<FJsonObject> JsonObject;
+		TSharedRef<TJsonReader<>> Reader = TJsonReaderFactory<>::Create(Response);
+
+		if (FJsonSerializer::Deserialize(Reader, JsonObject) && JsonObject.IsValid())
+		{
+			FString TransportState;
+			if (JsonObject->TryGetStringField(TEXT("Transport"), TransportState))
+			{
+				if (OnComplete)
+				{
+					OnComplete(true, TransportState);
+				}
+				return;
+			}
+		}
+
+		UE_LOG(LogTemp, Warning, TEXT("Failed to parse transport state from response: %s"), *Response);
+		if (OnComplete)
+		{
+			OnComplete(false, FString());
+		}
+	});
+}
+
+void ULiveLinkSoundDeviceBase::SetSetting(const FString& Key, const FString& Value, TFunction<void(bool bSuccess)> OnComplete)
+{
+	// URL encode the value
+	FString EncodedValue = FGenericPlatformHttp::UrlEncode(Value);
+	FString Command = FString::Printf(TEXT("setsetting/%s=%s"), *Key, *EncodedValue);
+
+	SendCommand(Command, [OnComplete](bool bSuccess, const FString& Response)
+	{
+		if (OnComplete)
+		{
+			OnComplete(bSuccess);
+		}
+	});
+}
+
+void ULiveLinkSoundDeviceBase::ConfigureDrivesForRecording()
+{
+	// Set all 4 drives to Record mode (not File Transfer)
+	for (int32 i = 1; i <= 4; i++)
+	{
+		FString Key = FString::Printf(TEXT("RecordToDrive%d"), i);
+		SetSetting(Key, TEXT("Record"), [](bool) {});
+	}
+}
+
+void ULiveLinkSoundDeviceBase::PollTransportState()
+{
+	if (bPollingRequestInFlight)
+	{
+		return;
+	}
+
+	// Check timeout
+	if (FPlatformTime::Seconds() - PollingStartTime > PollingTimeoutSeconds)
+	{
+		UE_LOG(LogTemp, Warning, TEXT("Transport state polling timeout"));
+		CurrentPollingOperation = EPollingOperation::None;
+		bPollingRequestInFlight = false;
+		return;
+	}
+
+	bPollingRequestInFlight = true;
+
+	TWeakObjectPtr<ULiveLinkSoundDeviceBase> WeakThis(this);
+	GetTransportState([WeakThis](bool bSuccess, const FString& StateText)
+	{
+		if (!WeakThis.IsValid())
+		{
+			return;
+		}
+
+		WeakThis->OnTransportStateReceived(bSuccess, StateText);
+	});
+}
+
+void ULiveLinkSoundDeviceBase::OnTransportStateReceived(bool bSuccess, const FString& StateText)
+{
+	bPollingRequestInFlight = false;
+
+	if (!bSuccess)
+	{
+		CurrentPollingOperation = EPollingOperation::None;
+		return;
+	}
+
+	bool bComplete = false;
+
+	if (CurrentPollingOperation == EPollingOperation::WaitingForRecording)
+	{
+		if (StateText.Equals(TEXT("rec"), ESearchCase::IgnoreCase))
+		{
+			bIsRecording = true;
+			bComplete = true;
+			UE_LOG(LogTemp, Log, TEXT("Recording started"));
+		}
+	}
+	else if (CurrentPollingOperation == EPollingOperation::WaitingForIdle)
+	{
+		if (StateText.Equals(TEXT("stop"), ESearchCase::IgnoreCase))
+		{
+			bIsRecording = false;
+			bComplete = true;
+			OnRecordingStopped();
+		}
+	}
+
+	if (bComplete)
+	{
+		CurrentPollingOperation = EPollingOperation::None;
+	}
+}
+
+void ULiveLinkSoundDeviceBase::OnRecordingStopped()
+{
+	UE_LOG(LogTemp, Log, TEXT("Recording stopped, retrieving file metadata..."));
+
+	TWeakObjectPtr<ULiveLinkSoundDeviceBase> WeakThis(this);
+
+	// Get recording file path
+	SendCommand(TEXT("invoke/RemoteApi/currentRecordTake()"), [WeakThis](bool bSuccess, const FString& Response)
+	{
+		if (!WeakThis.IsValid() || !bSuccess)
+		{
+			return;
+		}
+
+		// Parse JSON: {"String": "/HD1/..."}
+		TSharedPtr<FJsonObject> JsonObject;
+		TSharedRef<TJsonReader<>> Reader = TJsonReaderFactory<>::Create(Response);
+
+		if (FJsonSerializer::Deserialize(Reader, JsonObject) && JsonObject.IsValid())
+		{
+			FString Path;
+			if (JsonObject->TryGetStringField(TEXT("String"), Path))
+			{
+				// Normalize path
+				FString NormalizedPath = NormalizeFilePath(Path);
+				UE_LOG(LogTemp, Log, TEXT("Recording path: %s"), *NormalizedPath);
+
+				// Retrieve file metadata
+				WeakThis->RetrieveFileMetadata(Path);
+			}
+		}
+	});
+}
+
+void ULiveLinkSoundDeviceBase::RetrieveFileMetadata(const FString& FilePath)
+{
+	FString Command = FString::Printf(TEXT("filedetails%s"), *FilePath);
+
+	SendCommand(Command, [FilePath](bool bSuccess, const FString& Response)
+	{
+		if (!bSuccess)
+		{
+			return;
+		}
+
+		// Parse JSON: {"FileDetails": {"timecodeStart": "...", "duration": ...}}
+		TSharedPtr<FJsonObject> JsonObject;
+		TSharedRef<TJsonReader<>> Reader = TJsonReaderFactory<>::Create(Response);
+
+		if (FJsonSerializer::Deserialize(Reader, JsonObject) && JsonObject.IsValid())
+		{
+			const TSharedPtr<FJsonObject>* FileDetailsPtr;
+			if (JsonObject->TryGetObjectField(TEXT("FileDetails"), FileDetailsPtr))
+			{
+				const TSharedPtr<FJsonObject>& FileDetails = *FileDetailsPtr;
+
+				FString TimecodeStart;
+				double Duration = 0.0;
+
+				FileDetails->TryGetStringField(TEXT("timecodeStart"), TimecodeStart);
+				FileDetails->TryGetNumberField(TEXT("duration"), Duration);
+
+				UE_LOG(LogTemp, Log, TEXT("File metadata - TC Start: %s, Duration: %.2fs"),
+					*TimecodeStart, Duration);
+			}
+		}
+	});
+}
+
+FString ULiveLinkSoundDeviceBase::NormalizeFilePath(const FString& Path)
+{
+	FString Result = Path;
+
+	// Convert /HD{i}/ to /Drive_{i}/
+	for (int32 i = 1; i <= 4; i++)
+	{
+		FString OldPattern = FString::Printf(TEXT("/HD%d/"), i);
+		FString NewPattern = FString::Printf(TEXT("/Drive_%d/"), i);
+		Result = Result.Replace(*OldPattern, *NewPattern);
+	}
+
+	return Result;
+}
+
+bool ULiveLinkSoundDeviceBase::DeviceTick(float DeltaTime)
+{
+	// Handle periodic reconnection attempts when disconnected
+	if (ConnectionStatus == ELiveLinkDeviceConnectionStatus::Disconnected)
+	{
+		ReconnectionTimeAccumulator += DeltaTime;
+		if (ReconnectionTimeAccumulator >= ReconnectionIntervalSeconds)
+		{
+			ReconnectionTimeAccumulator = 0.0f;
+			AttemptReconnection();
+		}
+	}
+
+	// Handle transport state polling
+	if (CurrentPollingOperation != EPollingOperation::None)
+	{
+		PollTransportState();
+	}
+
+	return true; // Continue ticking
+}
+
+void ULiveLinkSoundDeviceBase::AttemptReconnection()
+{
+	if (ConnectionStatus == ELiveLinkDeviceConnectionStatus::Disconnected)
+	{
+		UE_LOG(LogTemp, Log, TEXT("Attempting reconnection to Sound Devices..."));
+		Connect_Implementation();
+	}
+}
+
+void ULiveLinkSoundDeviceBase::HandleSlateNameChanged(FStringView InSlateName)
+{
+	if (ConnectionStatus == ELiveLinkDeviceConnectionStatus::Connected)
+	{
+		FString Slate(InSlateName);
+		UE_LOG(LogTemp, Log, TEXT("Slate name changed to: %s"), *Slate);
+		SetSetting(TEXT("SceneName"), Slate, [](bool) {});
+	}
+}
+
+void ULiveLinkSoundDeviceBase::HandleTakeNumberChanged(int32 InTakeNumber)
+{
+	if (ConnectionStatus == ELiveLinkDeviceConnectionStatus::Connected)
+	{
+		UE_LOG(LogTemp, Log, TEXT("Take number changed to: %d"), InTakeNumber);
+		SetSetting(TEXT("TakeNumber"), FString::FromInt(InTakeNumber), [](bool) {});
+	}
+}
+
+#undef LOCTEXT_NAMESPACE
diff --git a/Engine/Plugins/VirtualProduction/LiveLinkSoundDevice/Source/LiveLinkSoundDevice/Private/LiveLinkSoundDeviceModule.cpp b/Engine/Plugins/VirtualProduction/LiveLinkSoundDevice/Source/LiveLinkSoundDevice/Private/LiveLinkSoundDeviceModule.cpp
new file mode 100644
--- /dev/null
+++ b/Engine/Plugins/VirtualProduction/LiveLinkSoundDevice/Source/LiveLinkSoundDevice/Private/LiveLinkSoundDeviceModule.cpp
@@ -0,0 +1,17 @@
+// Copyright Epic Games, Inc. All Rights Reserved.
+
+#include "Modules/ModuleManager.h"
+
+class FLiveLinkSoundDeviceModule : public IModuleInterface
+{
+public:
+	virtual void StartupModule() override
+	{
+	}
+
+	virtual void ShutdownModule() override
+	{
+	}
+};
+
+IMPLEMENT_MODULE(FLiveLinkSoundDeviceModule, LiveLinkSoundDevice)
diff --git a/Engine/Plugins/VirtualProduction/LiveLinkSoundDevice/Source/LiveLinkSoundDevice/Public/HttpDigestAuth.h b/Engine/Plugins/VirtualProduction/LiveLinkSoundDevice/Source/LiveLinkSoundDevice/Public/HttpDigestAuth.h
new file mode 100644
--- /dev/null
+++ b/Engine/Plugins/VirtualProduction/LiveLinkSoundDevice/Source/LiveLinkSoundDevice/Public/HttpDigestAuth.h
@@ -0,0 +1,71 @@
+// Copyright Epic Games, Inc. All Rights Reserved.
+
+#pragma once
+
+#include "CoreMinimal.h"
+
+/**
+ * Challenge data parsed from HTTP Digest WWW-Authenticate header
+ */
+struct FDigestChallenge
+{
+	FString Realm;
+	FString Nonce;
+	FString Qop;        // Quality of protection (e.g., "auth", "auth-int")
+	FString Opaque;
+	FString Algorithm;  // e.g., "MD5", "SHA-256"
+};
+
+/**
+ * Helper class for HTTP Digest Authentication (RFC 2617)
+ */
+class FHttpDigestAuthHelper
+{
+public:
+	/**
+	 * Parse WWW-Authenticate header from 401 response
+	 * @param AuthHeader The WWW-Authenticate header value
+	 * @param OutChallenge Parsed challenge data
+	 * @return true if parsing succeeded
+	 */
+	static bool ParseDigestChallenge(const FString& AuthHeader, FDigestChallenge& OutChallenge);
+
+	/**
+	 * Generate Authorization header value for authenticated request
+	 * @param Username Username for authentication
+	 * @param Password Password for authentication
+	 * @param Method HTTP method (GET, POST, etc.)
+	 * @param URI Request URI path
+	 * @param Challenge Challenge data from server
+	 * @return Authorization header value (includes "Digest " prefix)
+	 */
+	static FString GenerateDigestResponse(
+		const FString& Username,
+		const FString& Password,
+		const FString& Method,
+		const FString& URI,
+		const FDigestChallenge& Challenge);
+
+private:
+	/**
+	 * Compute MD5 hash of input string
+	 * @param Input String to hash
+	 * @return Lowercase hex string of MD5 digest
+	 */
+	static FString ComputeMD5Hash(const FString& Input);
+
+	/**
+	 * Generate random nonce for client
+	 * @return Hex string nonce
+	 */
+	static FString GenerateNonce();
+
+	/**
+	 * Parse key=value or key="value" pair from header
+	 * @param Pair String containing key=value
+	 * @param OutKey Extracted key
+	 * @param OutValue Extracted value (quotes removed)
+	 * @return true if parsing succeeded
+	 */
+	static bool ParseHeaderPair(const FString& Pair, FString& OutKey, FString& OutValue);
+};
diff --git a/Engine/Plugins/VirtualProduction/LiveLinkSoundDevice/Source/LiveLinkSoundDevice/Public/LiveLinkSoundDeviceBase.h b/Engine/Plugins/VirtualProduction/LiveLinkSoundDevice/Source/LiveLinkSoundDevice/Public/LiveLinkSoundDeviceBase.h
new file mode 100644
--- /dev/null
+++ b/Engine/Plugins/VirtualProduction/LiveLinkSoundDevice/Source/LiveLinkSoundDevice/Public/LiveLinkSoundDeviceBase.h
@@ -0,0 +1,154 @@
+// Copyright Epic Games, Inc. All Rights Reserved.
+
+#pragma once
+
+#include "CoreMinimal.h"
+#include "LiveLinkDevice.h"
+#include "LiveLinkDeviceCapability_Connection.h"
+#include "LiveLinkDeviceCapability_Recording.h"
+#include "LiveLinkSoundDeviceSettings.h"
+#include "HttpDigestAuth.h"
+#include "Interfaces/IHttpRequest.h"
+#include "Containers/Ticker.h"
+
+#include "LiveLinkSoundDeviceBase.generated.h"
+
+/**
+ * Concrete implementation for Sound Devices recording devices
+ * Provides both Connection and Recording capabilities via REST API with HTTP Digest Authentication
+ */
+UCLASS(BlueprintType, meta = (DisplayName = "Sound Devices Recorder", ToolTip = "Support for Sound Devices recorders (MixPre, 888, Scorpio)"))
+class LIVELINKSOUNDDEVICE_API ULiveLinkSoundDeviceBase : public ULiveLinkDevice
+	, public ILiveLinkDeviceCapability_Connection
+	, public ILiveLinkDeviceCapability_Recording
+{
+	GENERATED_BODY()
+
+public:
+	//~ Begin ULiveLinkDevice interface
+	virtual TSubclassOf<ULiveLinkDeviceSettings> GetSettingsClass() const override;
+	virtual FText GetDisplayName() const override;
+	virtual EDeviceHealth GetDeviceHealth() const override;
+	virtual FText GetHealthText() const override;
+	virtual void OnDeviceAdded() override;
+	virtual void OnDeviceRemoved() override;
+	virtual void OnSettingChanged(const FPropertyChangedEvent& InPropertyChangedEvent) override;
+	//~ End ULiveLinkDevice interface
+
+	//~ Begin ILiveLinkDeviceCapability_Connection interface
+	virtual ELiveLinkDeviceConnectionStatus GetConnectionStatus_Implementation() const override;
+	virtual FString GetHardwareId_Implementation() const override;
+	virtual bool SetHardwareId_Implementation(const FString& HardwareID) override;
+	virtual bool CanSetHardwareId_Implementation() override;
+	virtual bool Connect_Implementation() override;
+	virtual bool Disconnect_Implementation() override;
+	//~ End ILiveLinkDeviceCapability_Connection interface
+
+	//~ Begin ILiveLinkDeviceCapability_Recording interface
+	virtual bool StartRecording_Implementation() override;
+	virtual bool StopRecording_Implementation() override;
+	virtual bool IsRecording_Implementation() const override;
+	//~ End ILiveLinkDeviceCapability_Recording interface
+
+private:
+	/** Build base URL for Sound Devices REST API */
+	FString GetBaseUrl() const;
+
+	/** Build full command URL */
+	FString GetCommandUrl(const FString& Command) const;
+
+	/** Get URI path for authentication (without host/port) */
+	FString GetRequestURI(const FString& Command) const;
+
+	/** Send authenticated command to device */
+	void SendCommand(const FString& Command, TFunction<void(bool bSuccess, const FString& Response)> OnComplete);
+
+	/** Create HTTP request for a command */
+	TSharedRef<class IHttpRequest> CreateRequest(const FString& Command);
+
+	/** Retry request with Digest authentication */
+	void RetryWithAuth(const FString& Command, const FDigestChallenge& Challenge, TFunction<void(bool, const FString&)> OnComplete);
+
+	/** Get transport state from device */
+	void GetTransportState(TFunction<void(bool bSuccess, const FString& StateText)> OnComplete);
+
+	/** Set device parameter (slate, take, etc.) */
+	void SetSetting(const FString& Key, const FString& Value, TFunction<void(bool bSuccess)> OnComplete);
+
+	/** Ensure all drives are in Record mode (not File Transfer) */
+	void ConfigureDrivesForRecording();
+
+	/** Poll transport state until target state is reached */
+	void PollTransportState();
+
+	/** Called when transport state polling completes */
+	void OnTransportStateReceived(bool bSuccess, const FString& StateText);
+
+	/** Called when recording stop is confirmed */
+	void OnRecordingStopped();
+
+	/** Retrieve file metadata after recording */
+	void RetrieveFileMetadata(const FString& FilePath);
+
+	/** Normalize file path from HD{i} format to Drive_{i} format */
+	static FString NormalizeFilePath(const FString& Path);
+
+	/** Main device tick function (handles reconnection and polling) */
+	bool DeviceTick(float DeltaTime);
+
+	/** Attempt reconnection if disconnected */
+	void AttemptReconnection();
+
+	/** Handle slate name changes from recording session */
+	void HandleSlateNameChanged(FStringView InSlateName);
+
+	/** Handle take number changes from recording session */
+	void HandleTakeNumberChanged(int32 InTakeNumber);
+
+protected:
+	/** Settings class to use for Sound Devices */
+	UPROPERTY(EditDefaultsOnly, Category="Live Link Device")
+	TSubclassOf<ULiveLinkDeviceSettings> SettingsClass;
+
+	/** Current connection status */
+	UPROPERTY()
+	ELiveLinkDeviceConnectionStatus ConnectionStatus = ELiveLinkDeviceConnectionStatus::Disconnected;
+
+	/** Whether the device is currently recording */
+	UPROPERTY()
+	bool bIsRecording = false;
+
+	/** Ticker handle for main device tick (reconnection and polling) */
+	FTSTicker::FDelegateHandle DeviceTickerHandle;
+
+	/** What operation we're polling for */
+	enum class EPollingOperation
+	{
+		None,
+		WaitingForRecording,  // Polling for "rec" state
+		WaitingForIdle,       // Polling for "stop" state
+	};
+	EPollingOperation CurrentPollingOperation = EPollingOperation::None;
+
+	/** Time when polling started (for timeout) */
+	double PollingStartTime = 0.0;
+
+	/** Time accumulator for reconnection attempts */
+	float ReconnectionTimeAccumulator = 0.0f;
+
+	/** Flag to prevent polling request flood on slow connections */
+	bool bPollingRequestInFlight = false;
+
+	/** Cached Digest challenge for reuse (optimization) */
+	FDigestChallenge CachedChallenge;
+	bool bHasCachedChallenge = false;
+
+	/** Polling timeout in seconds */
+	static constexpr double PollingTimeoutSeconds = 4.0;
+
+	/** Polling interval in seconds */
+	static constexpr float PollingIntervalSeconds = 0.1f;
+
+	/** Reconnection attempt interval in seconds */
+	static constexpr float ReconnectionIntervalSeconds = 5.0f;
+};
diff --git a/Engine/Plugins/VirtualProduction/LiveLinkSoundDevice/Source/LiveLinkSoundDevice/Public/LiveLinkSoundDeviceSettings.h b/Engine/Plugins/VirtualProduction/LiveLinkSoundDevice/Source/LiveLinkSoundDevice/Public/LiveLinkSoundDeviceSettings.h
new file mode 100644
--- /dev/null
+++ b/Engine/Plugins/VirtualProduction/LiveLinkSoundDevice/Source/LiveLinkSoundDevice/Public/LiveLinkSoundDeviceSettings.h
@@ -0,0 +1,37 @@
+// Copyright Epic Games, Inc. All Rights Reserved.
+
+#pragma once
+
+#include "CoreMinimal.h"
+#include "LiveLinkDevice.h"
+#include "LiveLinkSoundDeviceSettings.generated.h"
+
+/**
+ * Settings for Sound Devices recording device
+ */
+UCLASS()
+class LIVELINKSOUNDDEVICE_API ULiveLinkSoundDeviceSettings : public ULiveLinkDeviceSettings
+{
+	GENERATED_BODY()
+
+public:
+	/** IP address of the Sound Devices device */
+	UPROPERTY(EditAnywhere, Category="Sound Devices", meta=(DisplayName="IP Address"))
+	FString IpAddress = TEXT("192.168.1.100");
+
+	/** HTTP port for REST API communication */
+	UPROPERTY(EditAnywhere, Category="Sound Devices", AdvancedDisplay)
+	uint16 Port = 80;
+
+	/** Username for HTTP Digest Authentication */
+	UPROPERTY(EditAnywhere, Category="Sound Devices", AdvancedDisplay)
+	FString Username = TEXT("guest");
+
+	/** Password for HTTP Digest Authentication */
+	UPROPERTY(EditAnywhere, Category="Sound Devices", AdvancedDisplay)
+	FString Password = TEXT("guest");
+
+	/** Device model retrieved from device (read-only) */
+	UPROPERTY(VisibleAnywhere, Category="Sound Devices", AdvancedDisplay, meta=(DisplayName="Device Model"))
+	FString DeviceModel;
+};
diff --git a/Engine/Plugins/VirtualProduction/LiveLinkSoundDevice/Source/Scripts/mock_sounddevice_README.md b/Engine/Plugins/VirtualProduction/LiveLinkSoundDevice/Source/Scripts/mock_sounddevice_README.md
new file mode 100644
--- /dev/null
+++ b/Engine/Plugins/VirtualProduction/LiveLinkSoundDevice/Source/Scripts/mock_sounddevice_README.md
@@ -0,0 +1,371 @@
+# Mock Sound Devices Recorder Simulator
+
+A Python-based HTTP server that simulates the REST API of a Sound Devices recorder for testing the Live Link Sound Device plugin without physical hardware.
+
+## Features
+
+- **HTTP Digest Authentication**: Full implementation of RFC 2617 Digest Auth
+- **Sound Devices REST API**: Simulates all key endpoints
+- **Transport State Machine**: Realistic `stop` ↔ `rec` state transitions
+- **Multi-Drive Support**: Simulates 4 drives with configurable modes
+- **Recording Metadata**: Generates timecode and duration information
+- **Real-time Logging**: Console output shows all commands received
+
+## Requirements
+
+- Python 3.6 or higher
+- No external dependencies (uses standard library only)
+
+## Usage
+
+### Basic Usage
+
+```bash
+# Run on port 8080 (recommended for testing)
+python mock_sounddevice_device.py --port 8080
+
+# Run on port 80 (requires admin/sudo)
+sudo python mock_sounddevice_device.py --port 80
+```
+
+### Command-Line Options
+
+```bash
+python mock_sounddevice_device.py [OPTIONS]
+
+Options:
+  --port PORT         HTTP port (default: 8080)
+  --ip IP            Bind IP address (default: 0.0.0.0)
+  --username USER    Auth username (default: guest)
+  --password PASS    Auth password (default: guest)
+```
+
+### Examples
+
+```bash
+# Listen on localhost only
+python mock_sounddevice_device.py --ip 127.0.0.1 --port 8080
+
+# Custom credentials
+python mock_sounddevice_device.py --username admin --password secret --port 8080
+
+# Production-like setup (requires elevated permissions)
+sudo python mock_sounddevice_device.py --port 80 --ip 192.168.1.100
+```
+
+## Testing with Live Link Hub
+
+1. **Start the Mock Server**:
+   ```bash
+   cd D:\work\fn\Engine\Plugins\VirtualProduction\LiveLinkDeviceSoundDevice\Source\Scripts
+   python mock_sounddevice_device.py --port 8080
+   ```
+
+   You should see:
+   ```
+   ============================================================
+   Mock Sound Devices Recorder Simulator
+   ============================================================
+   Listening on: http://0.0.0.0:8080
+   Authentication: guest / guest
+   Realm: Sound Devices API
+
+   Waiting for connections from Live Link Hub...
+   Press Ctrl+C to stop
+   ============================================================
+   ```
+
+2. **Launch Live Link Hub**:
+   - Start Unreal Editor or Live Link Hub standalone
+
+3. **Add Sound Devices Device**:
+   - Click "Add Device" in Live Link Hub
+   - Select "Sound Devices Recorder" from the device type list
+   - Device appears in the device table
+
+4. **Configure Device Settings**:
+   - Select the device row
+   - Edit settings in Details panel:
+     - **IP Address**: `127.0.0.1`
+     - **Port**: `8080`
+     - **Username**: `guest`
+     - **Password**: `guest`
+
+5. **Connect to Device**:
+   - Click "Connect" button in device row
+   - Device status should change: Disconnected → Connecting → Connected - Ready
+   - Mock server logs connection:
+     ```
+     [SETTING] FileNameFormat = Scene-Take
+     [SETTING] ReelName = 20260309
+     [SETTING] Drive 1 = Record
+     [SETTING] Drive 2 = Record
+     [SETTING] Drive 3 = Record
+     [SETTING] Drive 4 = Record
+     ```
+
+6. **Test Recording**:
+   - Set slate and take in Live Link Hub recording session
+   - Click "Start Recording" button
+   - Mock server logs:
+     ```
+     [SETTING] ReelName = 20260309
+     [SETTING] SceneName = MyScene
+     [SETTING] TakeNumber = 1
+     [SETTING] Drive 1 = Record
+     [SETTING] Drive 2 = Record
+     [SETTING] Drive 3 = Record
+     [SETTING] Drive 4 = Record
+
+     [RECORDING STARTED]
+       Path: /HD1/20260309/MyScene_Take1.mov
+       Scene: MyScene
+       Take: 1
+     ```
+   - Device health changes to "Good" (green)
+   - Health text shows "Recording"
+
+7. **Test Stop Recording**:
+   - Click "Stop Recording" button
+   - Mock server logs:
+     ```
+     [RECORDING STOPPED]
+       Path: /HD1/20260309/MyScene_Take1.mov
+     ```
+   - Device health returns to "Nominal"
+   - Health text shows "Connected - Ready"
+
+8. **Test Slate/Take Updates**:
+   - Change slate or take in Live Link Hub
+   - Mock server should log:
+     ```
+     [SETTING] SceneName = NewScene
+     [SETTING] TakeNumber = 5
+     ```
+
+9. **Test Disconnect**:
+   - Click "Disconnect" button
+   - Device status changes to "Disconnected"
+   - Health shows "Error" (red) - expected for disconnected state
+
+## API Endpoints
+
+The mock server implements these Sound Devices REST API endpoints:
+
+| Endpoint | Response | Description |
+|----------|----------|-------------|
+| `GET /sounddevices/tmcode` | Plain text: `"01:00:00:00"` | Get current timecode |
+| `GET /sounddevices/transport` | JSON: `{"Transport": "stop"\|"rec"}` | Get transport state |
+| `GET /sounddevices/settransport/rec` | `"OK"` | Start recording |
+| `GET /sounddevices/settransport/stop` | `"OK"` | Stop recording |
+| `GET /sounddevices/setsetting/{key}={value}` | `"OK"` | Set device parameter |
+| `GET /sounddevices/invoke/RemoteApi/currentRecordTake()` | JSON: `{"String": "/HD1/..."}` | Get recording path |
+| `GET /sounddevices/filedetails/{path}` | JSON: `{"FileDetails": {...}}` | Get file metadata |
+
+## HTTP Digest Authentication
+
+The mock server implements full HTTP Digest Authentication (RFC 2617):
+
+### Authentication Flow
+
+1. **Client Request** (no auth):
+   ```
+   GET /sounddevices/tmcode HTTP/1.1
+   ```
+
+2. **Server Challenge** (401 Unauthorized):
+   ```
+   HTTP/1.1 401 Unauthorized
+   WWW-Authenticate: Digest realm="Sound Devices API", nonce="abc123", qop="auth", opaque="xyz789"
+   ```
+
+3. **Client Response** (with computed digest):
+   ```
+   GET /sounddevices/tmcode HTTP/1.1
+   Authorization: Digest username="guest", realm="Sound Devices API", nonce="abc123",
+                         uri="/sounddevices/tmcode", qop=auth, nc=00000001, cnonce="def456",
+                         response="computed_md5_hash"
+   ```
+
+4. **Server Success**:
+   ```
+   HTTP/1.1 200 OK
+   Content-Type: text/plain
+
+   01:00:00:00
+   ```
+
+### Digest Computation
+
+The server verifies the response using MD5:
+
+```
+HA1 = MD5(username:realm:password)
+HA2 = MD5(method:uri)
+Response = MD5(HA1:nonce:nc:cnonce:qop:HA2)
+```
+
+## State Machine
+
+### Transport States
+
+```
+    stop ──┐
+     ↑     │
+     │     ↓
+     └──  rec
+```
+
+- **stop**: Idle, ready to record
+- **rec**: Currently recording
+
+### Drive Modes
+
+Each of the 4 drives can be in:
+- **Record**: Drive available for recording
+- **Ethernet File Transfer**: Drive in file transfer mode
+
+## Recording Metadata
+
+When recording stops, the mock server stores metadata:
+
+```json
+{
+  "FileDetails": {
+    "timecodeStart": "01:00:00:00",
+    "duration": 123.45,
+    "sampleRate": 48000,
+    "bitDepth": 24
+  }
+}
+```
+
+## File Path Format
+
+Recordings are stored with paths like:
+```
+/HD{drive_number}/{reel_name}/{scene_name}_Take{take_number}.mov
+```
+
+Example:
+```
+/HD1/20260309/MyScene_Take1.mov
+```
+
+The Live Link plugin normalizes these to:
+```
+/Drive_1/20260309/MyScene_Take1.mov
+```
+
+## Console Output
+
+The mock server provides detailed logging:
+
+```
+[SETTING] FileNameFormat = Scene-Take
+[SETTING] ReelName = 20260309
+[SETTING] SceneName = TestScene
+[SETTING] TakeNumber = 3
+[SETTING] Drive 1 = Record
+[SETTING] Drive 2 = Record
+[SETTING] Drive 3 = Record
+[SETTING] Drive 4 = Record
+
+[RECORDING STARTED]
+  Path: /HD1/20260309/TestScene_Take3.mov
+  Scene: TestScene
+  Take: 3
+
+[RECORDING STOPPED]
+  Path: /HD1/20260309/TestScene_Take3.mov
+```
+
+## Troubleshooting
+
+### Connection Failed
+
+If Live Link Hub cannot connect:
+
+1. **Check server is running**:
+   ```bash
+   python mock_sounddevice_device.py --port 8080
+   ```
+
+2. **Verify IP/port in Live Link Hub settings**:
+   - IP Address: `127.0.0.1` (localhost)
+   - Port: `8080`
+
+3. **Check firewall**:
+   - Ensure port 8080 is not blocked
+   - On Windows: Check Windows Defender Firewall
+
+4. **Check credentials**:
+   - Default: username=`guest`, password=`guest`
+   - Match mock server args if changed
+
+### Authentication Errors
+
+If you see 401 Unauthorized errors:
+
+1. **Verify credentials** match server configuration
+2. **Check console output** for authentication attempts
+3. **Ensure Digest Auth** is properly implemented in plugin
+
+### Recording Not Starting
+
+If recording doesn't start:
+
+1. **Check device is connected** (status = "Connected - Ready")
+2. **Verify transport state** in mock server console
+3. **Check for error messages** in Live Link Hub log
+
+## Limitations
+
+### Mock Server
+
+- Single-threaded (handles one request at a time)
+- No persistent storage (state resets on restart)
+- Simplified file metadata (fixed values)
+- No actual audio file generation
+- No multi-client support
+
+### Not Implemented
+
+- File transfer mode functionality
+- mDNS/Bonjour discovery
+- Individual track configuration
+- Real timecode generation
+- Advanced recording formats
+
+## Real Device Testing
+
+After testing with the mock server, test with actual Sound Devices hardware:
+
+1. **Supported Devices**:
+   - MixPre series (MixPre-3, MixPre-6, MixPre-10)
+   - 8-Series (888, 833, 888)
+   - Scorpio
+
+2. **Network Configuration**:
+   - Connect device to network
+   - Note device IP address (check device menu)
+   - Ensure device API is enabled
+
+3. **Live Link Hub Configuration**:
+   - Use device's actual IP address
+   - Port is typically `80`
+   - Default credentials: `guest` / `guest`
+
+4. **Verify API Access**:
+   ```bash
+   # Test from command line
+   curl -u guest:guest --digest http://{device_ip}/sounddevices/tmcode
+   ```
+
+## Support
+
+For issues or questions:
+- Check Live Link Hub logs
+- Review mock server console output
+- Verify network connectivity
+- Ensure Sound Devices firmware is up to date (for real devices)
diff --git a/Engine/Plugins/VirtualProduction/LiveLinkSoundDevice/Source/Scripts/mock_sounddevice_device.py b/Engine/Plugins/VirtualProduction/LiveLinkSoundDevice/Source/Scripts/mock_sounddevice_device.py
new file mode 100644
--- /dev/null
+++ b/Engine/Plugins/VirtualProduction/LiveLinkSoundDevice/Source/Scripts/mock_sounddevice_device.py
@@ -0,0 +1,347 @@
+#!/usr/bin/env python3
+"""
+Mock Sound Devices Recorder Simulator
+
+Simulates the REST API of a Sound Devices recorder for testing the Live Link Sound Device plugin.
+Implements HTTP Digest Authentication and the key endpoints needed for recording control.
+
+Usage:
+    python mock_sounddevice_device.py [--port PORT] [--ip IP]
+
+Example:
+    python mock_sounddevice_device.py --port 8080 --ip 127.0.0.1
+"""
+
+import argparse
+import json
+import hashlib
+import secrets
+import time
+import sys
+from http.server import HTTPServer, BaseHTTPRequestHandler
+from urllib.parse import urlparse, unquote, parse_qs
+from threading import Lock
+
+
+class SoundDeviceState:
+    """Manages the state of the mock Sound Devices recorder"""
+
+    def __init__(self):
+        self.lock = Lock()
+
+        # Device state
+        self.transport_state = "stop"  # "rec", "stop"
+        self.scene_name = "DefaultScene"
+        self.take_number = "1"
+        self.reel_name = "20260101"
+        self.file_name_format = "Scene-Take"
+        self.current_recording_path = ""
+
+        # Drive configuration (1-4)
+        self.drive_modes = {
+            1: "Record",
+            2: "Record",
+            3: "Record",
+            4: "Record"
+        }
+
+        # Recording metadata storage
+        self.recordings = {}  # path -> {timecodeStart, duration}
+
+    def start_recording(self):
+        """Start recording"""
+        with self.lock:
+            if self.transport_state == "stop":
+                self.transport_state = "rec"
+
+                # Generate recording path based on format
+                if self.file_name_format == "Scene-Take":
+                    filename = f"{self.scene_name}_Take{self.take_number}.mov"
+                else:
+                    filename = f"recording_{int(time.time())}.mov"
+
+                # Place on first drive in record mode
+                drive_num = 1
+                for d, mode in self.drive_modes.items():
+                    if mode == "Record":
+                        drive_num = d
+                        break
+
+                self.current_recording_path = f"/HD{drive_num}/{self.reel_name}/{filename}"
+
+                print(f"\n[RECORDING STARTED]")
+                print(f"  Path: {self.current_recording_path}")
+                print(f"  Scene: {self.scene_name}")
+                print(f"  Take: {self.take_number}")
+                sys.stdout.flush()
+
+                return True
+            return False
+
+    def stop_recording(self):
+        """Stop recording"""
+        with self.lock:
+            if self.transport_state == "rec":
+                self.transport_state = "stop"
+
+                # Store recording metadata
+                if self.current_recording_path:
+                    self.recordings[self.current_recording_path] = {
+                        "timecodeStart": "01:00:00:00",
+                        "duration": 123.45,
+                        "sampleRate": 48000,
+                        "bitDepth": 24
+                    }
+
+                print(f"\n[RECORDING STOPPED]")
+                print(f"  Path: {self.current_recording_path}")
+                sys.stdout.flush()
+
+                return True
+            return False
+
+    def set_setting(self, key, value):
+        """Set a device setting"""
+        with self.lock:
+            if key == "FileNameFormat":
+                self.file_name_format = value
+                print(f"[SETTING] FileNameFormat = {value}")
+                sys.stdout.flush()
+            elif key == "ReelName":
+                self.reel_name = value
+                print(f"[SETTING] ReelName = {value}")
+                sys.stdout.flush()
+            elif key == "SceneName":
+                self.scene_name = value
+                print(f"[SETTING] SceneName = {value}")
+                sys.stdout.flush()
+            elif key == "TakeNumber":
+                self.take_number = value
+                print(f"[SETTING] TakeNumber = {value}")
+                sys.stdout.flush()
+            elif key.startswith("RecordToDrive"):
+                drive_num = int(key[-1])  # Extract drive number
+                self.drive_modes[drive_num] = value
+                print(f"[SETTING] Drive {drive_num} = {value}")
+                sys.stdout.flush()
+            else:
+                print(f"[SETTING] {key} = {value} (unhandled)")
+                sys.stdout.flush()
+
+
+class DigestAuthHandler:
+    """Implements HTTP Digest Authentication (RFC 2617)"""
+
+    def __init__(self, username="guest", password="guest"):
+        self.username = username
+        self.password = password
+        self.realm = "Sound Devices API"
+        self.nonce = secrets.token_hex(16)
+        self.opaque = secrets.token_hex(16)
+
+    def generate_challenge(self):
+        """Generate WWW-Authenticate header"""
+        return f'Digest realm="{self.realm}", nonce="{self.nonce}", qop="auth", opaque="{self.opaque}"'
+
+    def verify_response(self, auth_header, method, uri):
+        """Verify Authorization header"""
+        if not auth_header or not auth_header.startswith("Digest "):
+            return False
+
+        # Parse Authorization header
+        auth_str = auth_header[7:]  # Remove "Digest "
+        parts = {}
+
+        # Simple parser for key=value or key="value"
+        for part in auth_str.split(","):
+            part = part.strip()
+            if "=" in part:
+                key, value = part.split("=", 1)
+                key = key.strip()
+                value = value.strip().strip('"')
+                parts[key] = value
+
+        # Verify required fields
+        if "response" not in parts or "nonce" not in parts:
+            return False
+
+        # Compute expected response
+        # HA1 = MD5(username:realm:password)
+        ha1_input = f"{self.username}:{self.realm}:{self.password}"
+        ha1 = hashlib.md5(ha1_input.encode()).hexdigest()
+
+        # HA2 = MD5(method:uri)
+        ha2_input = f"{method}:{uri}"
+        ha2 = hashlib.md5(ha2_input.encode()).hexdigest()
+
+        # Response
+        if "qop" in parts and parts["qop"] == "auth":
+            # qop=auth mode
+            nc = parts.get("nc", "00000001")
+            cnonce = parts.get("cnonce", "")
+            response_input = f"{ha1}:{parts['nonce']}:{nc}:{cnonce}:{parts['qop']}:{ha2}"
+        else:
+            # Legacy mode
+            response_input = f"{ha1}:{parts['nonce']}:{ha2}"
+
+        expected_response = hashlib.md5(response_input.encode()).hexdigest()
+
+        return parts["response"] == expected_response
+
+
+class SoundDeviceRequestHandler(BaseHTTPRequestHandler):
+    """HTTP request handler for Sound Devices API"""
+
+    # Shared state and auth handler (set by server)
+    device_state = None
+    auth_handler = None
+
+    def log_message(self, format, *args):
+        """Suppress default logging"""
+        pass
+
+    def do_GET(self):
+        """Handle GET requests"""
+        # Check authentication
+        auth_header = self.headers.get("Authorization")
+
+        if not auth_header:
+            # Send 401 with challenge
+            self.send_response(401)
+            self.send_header("WWW-Authenticate", self.auth_handler.generate_challenge())
+            self.send_header("Content-Length", "0")
+            self.end_headers()
+            return
+
+        # Verify authentication
+        if not self.auth_handler.verify_response(auth_header, "GET", self.path):
+            self.send_error(401, "Unauthorized")
+            return
+
+        # Route to command handler
+        parsed = urlparse(self.path)
+        path_parts = parsed.path.split("/")
+
+        if len(path_parts) >= 2 and path_parts[1] == "sounddevices":
+            # Extract command from path
+            command = "/".join(path_parts[2:])
+            self.handle_command(command)
+        else:
+            self.send_error(404, "Not Found")
+
+    def handle_command(self, command):
+        """Route command to appropriate handler"""
+        # Timecode
+        if command == "tmcode":
+            self.send_text_response("01:00:00:00")
+
+        # Transport state
+        elif command == "transport":
+            self.send_json_response({"Transport": self.device_state.transport_state})
+
+        # Transport control
+        elif command == "settransport/rec":
+            self.device_state.start_recording()
+            self.send_text_response("OK")
+
+        elif command == "settransport/stop":
+            self.device_state.stop_recording()
+            self.send_text_response("OK")
+
+        # Settings
+        elif command.startswith("setsetting/"):
+            param = command[11:]  # Remove "setsetting/"
+            if "=" in param:
+                key, value = param.split("=", 1)
+                value = unquote(value)  # URL decode
+                self.device_state.set_setting(key, value)
+            self.send_text_response("OK")
+
+        # Recording info
+        elif command == "invoke/RemoteApi/currentRecordTake()":
+            path = self.device_state.current_recording_path or "/HD1/default.mov"
+            self.send_json_response({"String": path})
+
+        # File details
+        elif command.startswith("filedetails/"):
+            file_path = command[12:]  # Remove "filedetails/"
+            if file_path in self.device_state.recordings:
+                metadata = self.device_state.recordings[file_path]
+                self.send_json_response({"FileDetails": metadata})
+            else:
+                # Return generic metadata for unknown files
+                self.send_json_response({
+                    "FileDetails": {
+                        "timecodeStart": "01:00:00:00",
+                        "duration": 60.0,
+                        "sampleRate": 48000,
+                        "bitDepth": 24
+                    }
+                })
+
+        else:
+            self.send_error(404, f"Unknown command: {command}")
+
+    def send_text_response(self, text):
+        """Send plain text response"""
+        response = text.encode()
+        self.send_response(200)
+        self.send_header("Content-Type", "text/plain")
+        self.send_header("Content-Length", str(len(response)))
+        self.end_headers()
+        self.wfile.write(response)
+
+    def send_json_response(self, data):
+        """Send JSON response"""
+        response = json.dumps(data).encode()
+        self.send_response(200)
+        self.send_header("Content-Type", "application/json")
+        self.send_header("Content-Length", str(len(response)))
+        self.end_headers()
+        self.wfile.write(response)
+
+
+def main():
+    parser = argparse.ArgumentParser(description="Mock Sound Devices Recorder Simulator")
+    parser.add_argument("--port", type=int, default=8080, help="HTTP port (default: 8080)")
+    parser.add_argument("--ip", type=str, default="0.0.0.0", help="Bind IP address (default: 0.0.0.0)")
+    parser.add_argument("--username", type=str, default="guest", help="Auth username (default: guest)")
+    parser.add_argument("--password", type=str, default="guest", help="Auth password (default: guest)")
+    args = parser.parse_args()
+
+    # Force unbuffered output for immediate console visibility
+    sys.stdout.reconfigure(line_buffering=True)
+    sys.stderr.reconfigure(line_buffering=True)
+
+    # Create shared state and auth handler
+    device_state = SoundDeviceState()
+    auth_handler = DigestAuthHandler(username=args.username, password=args.password)
+
+    # Assign to handler class
+    SoundDeviceRequestHandler.device_state = device_state
+    SoundDeviceRequestHandler.auth_handler = auth_handler
+
+    # Create and start server
+    server = HTTPServer((args.ip, args.port), SoundDeviceRequestHandler)
+
+    print("=" * 60)
+    print("Mock Sound Devices Recorder Simulator")
+    print("=" * 60)
+    print(f"Listening on: http://{args.ip}:{args.port}")
+    print(f"Authentication: {args.username} / {args.password}")
+    print(f"Realm: {auth_handler.realm}")
+    print()
+    print("Waiting for connections from Live Link Hub...")
+    print("Press Ctrl+C to stop")
+    print("=" * 60)
+    print()
+
+    try:
+        server.serve_forever()
+    except KeyboardInterrupt:
+        print("\n\nShutting down...")
+        server.shutdown()
+
+
+if __name__ == "__main__":
+    main()
diff --git a/Engine/Plugins/VirtualProduction/LiveLinkSoundDevice/Source/Scripts/run_mock.bat b/Engine/Plugins/VirtualProduction/LiveLinkSoundDevice/Source/Scripts/run_mock.bat
new file mode 100755
--- /dev/null
+++ b/Engine/Plugins/VirtualProduction/LiveLinkSoundDevice/Source/Scripts/run_mock.bat
@@ -0,0 +1,20 @@
+@echo off
+REM Helper script to run the mock Sound Devices server with visible console output
+REM Usage: run_mock.bat [port]
+REM Example: run_mock.bat 8080
+
+setlocal
+set PORT=8080
+
+if not "%1"=="" set PORT=%1
+
+echo ================================================================
+echo Starting Mock Sound Devices Recorder on port %PORT%
+echo ================================================================
+echo.
+echo Press Ctrl+C to stop the server
+echo.
+
+python -u mock_sounddevice_device.py --port %PORT%
+
+endlocal
diff --git a/Engine/Source/Programs/LiveLinkHub/Source/LiveLinkHub.Target.cs b/Engine/Source/Programs/LiveLinkHub/Source/LiveLinkHub.Target.cs
--- a/Engine/Source/Programs/LiveLinkHub/Source/LiveLinkHub.Target.cs
+++ b/Engine/Source/Programs/LiveLinkHub/Source/LiveLinkHub.Target.cs
@@ -55,12 +55,13 @@ public class LiveLinkHubTarget : TargetRules
 			"AppleARKitFaceSupport",
 			"XInputDevice",
 
+			// Live Link Hub Default Devices
 			"LiveLinkDevice",
 			"LiveLinkUnrealDevice",
 			"LiveLinkKiProDevice",
+			"LiveLinkSoundDevice",
 
 			"LiveLinkOpenVR",
-
 			"LiveLinkOpenTrackIO",
 
 			"CaptureManagerApp",
